DMX - Enable the DMX Control Console plugin for the DMX Template project

diff --git a/Templates/TP_DMXBP/TP_DMXBP.uproject b/Templates/TP_DMXBP/TP_DMXBP.uproject
--- a/Templates/TP_DMXBP/TP_DMXBP.uproject
+++ b/Templates/TP_DMXBP/TP_DMXBP.uproject
@@ -5,7 +5,7 @@
 	"Description": "",
 	"Plugins": [
 		{
-			"Name": "DMXProtocol",
+			"Name": "DMXControlConsole",
 			"Enabled": true
 		},
 		{
@@ -20,6 +20,10 @@
 			"Name": "DMXPixelMapping",
 			"Enabled": true
 		},
+		{
+			"Name": "DMXProtocol",
+			"Enabled": true
+		},
 		{
 			"Name": "OculusVR",
 			"Enabled": false,
